Fix shader compilation for SM5 tvOS

diff --git a/Engine/Config/TVOS/DataDrivenPlatformInfo.ini b/Engine/Config/TVOS/DataDrivenPlatformInfo.ini
--- a/Engine/Config/TVOS/DataDrivenPlatformInfo.ini
+++ b/Engine/Config/TVOS/DataDrivenPlatformInfo.ini
@@ -27,6 +27,8 @@ bSupports4ComponentUAVReadWrite=true
 bSupportsManualVertexFetch=false
 bSupportsDualSourceBlending=true
 bSupportsIndependentSamplers=true
+bSupportsDistanceFields=true
+bSupportsLumenGI=true
 FriendlyName=LOCTEXT("FriendlyShaderPlatformName_TVOS_SM5", "TVOS Metal SM5")
 
 [ShaderPlatform METAL_ES3_1_TVOS]
